Fix deprecation warning.

#rnx

diff --git a/Engine/Plugins/Experimental/CodeView/Source/CodeView/CodeView.Build.cs b/Engine/Plugins/Experimental/CodeView/Source/CodeView/CodeView.Build.cs
--- a/Engine/Plugins/Experimental/CodeView/Source/CodeView/CodeView.Build.cs
+++ b/Engine/Plugins/Experimental/CodeView/Source/CodeView/CodeView.Build.cs
@@ -26,6 +26,7 @@ namespace UnrealBuildTool.Rules
 				{
 					"DetailCustomizations",
 					"SlateCore",
+					"SourceCodeAccess"
 				}
 			);
 		}
diff --git a/Engine/Plugins/Experimental/CodeView/Source/CodeView/Private/CodeViewPlugin.cpp b/Engine/Plugins/Experimental/CodeView/Source/CodeView/Private/CodeViewPlugin.cpp
--- a/Engine/Plugins/Experimental/CodeView/Source/CodeView/Private/CodeViewPlugin.cpp
+++ b/Engine/Plugins/Experimental/CodeView/Source/CodeView/Private/CodeViewPlugin.cpp
@@ -9,6 +9,8 @@
 #include "DetailCategoryBuilder.h"
 #include "SCodeView.h"
 #include "DesktopPlatformModule.h"
+#include "ISourceCodeAccessor.h"
+#include "ISourceCodeAccessModule.h"
 
 class FCodeViewPlugin : public IModuleInterface
 {
@@ -27,8 +29,8 @@ public:
 
 	void AddCodeViewCategory(IDetailLayoutBuilder& DetailBuilder, const FGetSelectedActors& GetSelectedActors)
 	{
-		FString SolutionPath;
-		if(FDesktopPlatformModule::Get()->GetSolutionPath(SolutionPath))
+		ISourceCodeAccessModule* SourceCodeAccessModule = FModuleManager::GetModulePtr<ISourceCodeAccessModule>("SourceCodeAccess");
+		if(SourceCodeAccessModule != nullptr && SourceCodeAccessModule->GetAccessor().CanAccessSourceCode())
 		{
 			TSharedRef< CodeView::SCodeView > CodeViewWidget =
 				SNew( CodeView::SCodeView )
